Make dependencyy from UndoHistory to UnrealEd optional

diff --git a/Engine/Source/Developer/UndoHistory/Private/UndoHistoryModule.cpp b/Engine/Source/Developer/UndoHistory/Private/UndoHistoryModule.cpp
new file mode 100644
--- /dev/null
+++ b/Engine/Source/Developer/UndoHistory/Private/UndoHistoryModule.cpp
@@ -0,0 +1,22 @@
+// Copyright Epic Games, Inc. All Rights Reserved.
+
+#include "UndoHistoryModule.h"
+
+#define LOCTEXT_NAMESPACE "FUndoHistoryModule"
+
+void FUndoHistoryModule::StartupModule()
+{}
+
+void FUndoHistoryModule::ShutdownModule()
+{}
+
+bool FUndoHistoryModule::SupportsDynamicReloading()
+{
+	return true;
+}
+
+
+IMPLEMENT_MODULE(FUndoHistoryModule, UndoHistory);
+
+
+#undef LOCTEXT_NAMESPACE
diff --git a/Engine/Source/Developer/UndoHistory/Private/UndoHistoryModule.h b/Engine/Source/Developer/UndoHistory/Private/UndoHistoryModule.h
new file mode 100644
--- /dev/null
+++ b/Engine/Source/Developer/UndoHistory/Private/UndoHistoryModule.h
@@ -0,0 +1,20 @@
+// Copyright Epic Games, Inc. All Rights Reserved.
+
+#pragma once
+
+#include "CoreMinimal.h"
+#include "IUndoHistoryModule.h"
+
+/**
+* Implements the UndoHistory module.
+*/
+class FUndoHistoryModule : public IUndoHistoryModule
+{
+public:
+
+	//~ Begin IModuleInterface Interface
+	virtual void StartupModule() override;
+	virtual void ShutdownModule() override;
+	virtual bool SupportsDynamicReloading() override;
+	//~ End IModuleInterface Interface
+};
diff --git a/Engine/Source/Developer/UndoHistory/Private/Widgets/SUndoHistoryDetails.cpp b/Engine/Source/Developer/UndoHistory/Private/Widgets/SUndoHistoryDetails.cpp
new file mode 100644
--- /dev/null
+++ b/Engine/Source/Developer/UndoHistory/Private/Widgets/SUndoHistoryDetails.cpp
@@ -0,0 +1,490 @@
+// Copyright Epic Games, Inc. All Rights Reserved.
+
+#include "Widgets/SUndoHistoryDetails.h"
+#include "SlateOptMacros.h"
+#include "Widgets/Layout/SSeparator.h"
+#include "Widgets/Text/STextBlock.h"
+#include "Widgets/Views/STableRow.h"
+#include "Widgets/Views/SHeaderRow.h"
+#include "Widgets/Views/STreeView.h"
+#include "Widgets/Layout/SGridPanel.h"
+#include "Widgets/Input/SHyperlink.h"
+#include "Widgets/SToolTip.h"
+#include "Widgets/Input/SCheckBox.h"
+#include "Widgets/Images/SImage.h"
+#include "Widgets/Colors/SColorBlock.h"
+#include "Widgets/Input/SSearchBox.h"
+#include "Widgets/Layout/SSeparator.h"
+#include "HAL/PlatformApplicationMisc.h"
+#include "UndoHistoryUtils.h"
+#include "Misc/ITransaction.h"
+#include "Algo/Transform.h"
+
+#define LOCTEXT_NAMESPACE "SUndoHistoryDetails"
+
+// Static Transaction Details UI FNames
+namespace TransactionDetailsUI
+{
+	static const FName NameLabel(TEXT("Name"));
+	static const FName TypeLabel(TEXT("Type"));
+	static const FName ChangeFlagsLabel(TEXT("ChangeFlags"));
+ 
+	TSharedRef<SImage> CreateRenameIcon()
+	{
+		return SNew(SImage)
+			.Image(FAppStyle::Get().GetBrush("Icons.Rename"))
+			.ToolTipText(LOCTEXT("RenameChangeToolTip", "Renamed."))
+			.ColorAndOpacity(FSlateColor::UseForeground()); 
+	}
+
+	TSharedRef<SImage> CreateOuterChangeIcon()
+	{
+		return SNew(SImage)
+			.Image(FAppStyle::Get().GetBrush("Icons.CircleArrowUp"))
+ 
+			.ToolTipText(LOCTEXT("OuterChangeToolTip", "Has an outer change."))
+			.ColorAndOpacity(FSlateColor::UseForeground()); 
+	}
+	TSharedRef<SImage> CreatePendingKillIcon()
+	{
+		return SNew(SImage)
+			.Image(FAppStyle::Get().GetBrush("Icons.FilledCircle"))
+ 
+			.ToolTipText(LOCTEXT("PendingKillToolTip", "Has a pending kill change."))
+			.ColorAndOpacity(FSlateColor::UseForeground());
+	}
+	TSharedRef<SImage> CreateNonPropertyChangeIcon()
+	{
+		return SNew(SImage)
+			.Image(FAppStyle::Get().GetBrush("Icons.Settings"))
+ 
+			.ToolTipText(LOCTEXT("NonPropertyChangeToolTip", "Has a non-property change."))
+			.ColorAndOpacity(FSlateColor::UseForeground());
+	}
+	TSharedRef<SImage> CreateAnnotationIcon()
+	{
+		return SNew(SImage)
+			.Image(FAppStyle::Get().GetBrush("Icons.Info"))
+ 
+			.ToolTipText(LOCTEXT("AnnotationToolTip", "Has an annotation."))
+			.ColorAndOpacity(FSlateColor::UseForeground()); 
+	}
+};
+
+/**
+ * Implements a row widget for the UndoDetails tree view.
+ */
+class SUndoHistoryDetailsRow
+	: public SMultiColumnTableRow<TSharedPtr<int32> >
+{
+
+public:
+
+	SLATE_BEGIN_ARGS(SUndoHistoryDetailsRow)
+	: _TransactionEvent(nullptr)
+	{ }
+		SLATE_ARGUMENT(TSharedPtr<FTransactionObjectEvent>, TransactionEvent)
+		SLATE_ARGUMENT(FString, Name)
+		SLATE_ARGUMENT(FString, Type)
+		SLATE_ATTRIBUTE(FText, FilterText)
+	SLATE_END_ARGS()
+
+public:
+
+	/**
+	 * Constructs the widget.
+	 *
+	 * @param InArgs The construction arguments.
+	 */
+	void Construct(const FArguments& InArgs, const TSharedRef<STableViewBase>& InOwnerTableView)
+	{
+		TransactionEvent = InArgs._TransactionEvent;
+		Name = InArgs._Name;
+		Type = InArgs._Type;
+		FilterText = InArgs._FilterText;
+
+		SMultiColumnTableRow<TSharedPtr<int32> >::Construct(FSuperRowType::FArguments(), InOwnerTableView);
+	}
+
+public:
+
+	BEGIN_SLATE_FUNCTION_BUILD_OPTIMIZATION
+	virtual TSharedRef<SWidget> GenerateWidgetForColumn(const FName& ColumnName) override
+	{
+		if (ColumnName == TransactionDetailsUI::NameLabel)
+		{
+			return
+				SNew(SHorizontalBox)
+					+ SHorizontalBox::Slot()
+					.AutoWidth()
+					[
+						SNew(SExpanderArrow, SharedThis(this))
+							.IndentAmount(8)
+							.ShouldDrawWires(true)
+					]
+
+					+ SHorizontalBox::Slot()
+					.FillWidth(1.0f)
+					[
+						SNew(STextBlock)
+							.Text(FText::FromString(Name))
+							.HighlightText(FilterText)
+					];
+		}
+		else if (ColumnName == TransactionDetailsUI::TypeLabel)
+		{
+			return SNew(STextBlock)
+				.Text(FText::FromString(Type));
+		}
+		else if (ColumnName == TransactionDetailsUI::ChangeFlagsLabel)
+		{
+			if (!TransactionEvent.IsValid())
+			{
+				return SNullWidget::NullWidget;
+			}
+
+			return SNew(SBox)
+				.HAlign(HAlign_Center)
+				.VAlign(VAlign_Center)
+				.Padding(FMargin(2.0f, 0.0f))
+				[
+					CreateFlagBox()
+				];
+		}
+
+		return SNullWidget::NullWidget;
+	}
+	END_SLATE_FUNCTION_BUILD_OPTIMIZATION
+
+private:
+	TSharedRef<SHorizontalBox> CreateFlagBox()
+	{
+		TSharedRef<SHorizontalBox> HorizontalBox = SNew(SHorizontalBox);
+
+		if (TransactionEvent->HasNameChange())
+		{
+			HorizontalBox->AddSlot()
+				.Padding(FMargin(2.0f, 0.0f))
+				[
+					TransactionDetailsUI::CreateRenameIcon()
+				];
+		}
+
+		if (TransactionEvent->HasOuterChange())
+		{
+			HorizontalBox->AddSlot()
+				.Padding(FMargin(2.0f, 0.0f))
+				[
+					TransactionDetailsUI::CreateOuterChangeIcon()
+				];
+		}
+
+		if (TransactionEvent->HasPendingKillChange())
+		{
+			HorizontalBox->AddSlot()
+				.Padding(FMargin(2.0f, 0.0f))
+				[
+					TransactionDetailsUI::CreatePendingKillIcon()
+				];
+		}
+
+		if (TransactionEvent->HasNonPropertyChanges())
+		{
+			HorizontalBox->AddSlot()
+				.Padding(FMargin(2.0f, 0.0f))
+				[
+					TransactionDetailsUI::CreateNonPropertyChangeIcon()
+				];
+		}
+
+		if (TransactionEvent->GetAnnotation().IsValid())
+		{
+			HorizontalBox->AddSlot()
+				.Padding(FMargin(2.0f, 0.0f))
+				[
+					TransactionDetailsUI::CreateAnnotationIcon()
+				];
+		}
+
+		return HorizontalBox;
+	}
+
+private:
+
+	TSharedPtr<FTransactionObjectEvent> TransactionEvent;
+	FString Name;
+	FString Type;
+	TAttribute<FText> FilterText;
+
+};
+
+void SUndoHistoryDetails::Construct(const FArguments& InArgs)
+{
+	const auto OnGetChildren = [](FUndoDetailsTreeNodePtr InNode, TArray<FUndoDetailsTreeNodePtr> & OutChildren)
+	{
+		OutChildren = InNode->Children;
+	};
+
+	auto FilterDelegate = FTreeItemTextFilter::FItemToStringArray::CreateSP(this, &SUndoHistoryDetails::PopulateSearchStrings);
+	SearchBoxFilter = MakeShared<FTreeItemTextFilter>(FilterDelegate);
+	SearchBoxFilter->OnChanged().AddSP(this, &SUndoHistoryDetails::FullRefresh);
+	
+	bNeedsRefresh = true;
+	bNeedsExpansion = false;
+
+	ChildSlot
+		[
+			SNew(SBorder)
+				.BorderImage(FAppStyle::Get().GetBrush("ToolPanel.GroupBorder"))
+				.Padding(FMargin(4.0f, 1.0f))
+				[
+					SNew(SVerticalBox)
+						.Clipping(EWidgetClipping::ClipToBounds)
+						.Visibility(this, &SUndoHistoryDetails::HandleDetailsVisibility)
+
+					+ SVerticalBox::Slot()
+						.Padding(FMargin(0.0f, 4.0f))
+						.AutoHeight()
+						[
+							SNew(SBorder)
+								.BorderImage(FAppStyle::Get().GetBrush("ToolPanel.GroupBorder"))
+								.Padding(FMargin(0.0f, 1.0f))
+								[
+									SNew(SGridPanel)
+
+									+ SGridPanel::Slot(0, 0)
+									[
+										SNew(STextBlock)
+											.Font(FAppStyle::Get().GetFontStyle("NormalFontBold"))
+											.Text(LOCTEXT("TransactionActionLabel", "Action:"))
+									]
+
+									+ SGridPanel::Slot(1, 0)
+										.Padding(16.0f, 0.0f, 8.0f, 0.0f)
+										[
+											SNew(STextBlock)
+											.Text(this, &SUndoHistoryDetails::HandleTransactionName)
+										]
+
+									+ SGridPanel::Slot(0, 1)
+										.Padding(0.0f, 4.0f, 0.0f, 0.0f)
+										[
+											SNew(STextBlock)
+												.Font(FAppStyle::Get().GetFontStyle("NormalFontBold"))
+												.Text(LOCTEXT("TransactionIdLabel", "Transaction Id:"))
+										]
+
+									+ SGridPanel::Slot(1, 1)
+										.Padding(16.0f, 4.0f, 8.0f, 0.0f)
+										[
+											SNew(SHyperlink)
+												.ToolTipText(NSLOCTEXT("SUndoHistoryDetails", "ClickToCopy", "Click to copy ID."))
+												.Text(this, &SUndoHistoryDetails::HandleTransactionId)
+												.OnNavigate(this, &SUndoHistoryDetails::HandleTransactionIdNavigate)
+										]
+								]
+						]
+
+					+ SVerticalBox::Slot()
+						.AutoHeight()
+						.Padding(FMargin(0.0f, 5.0f))
+						[
+							SNew(SSeparator)
+								.Orientation(EOrientation::Orient_Horizontal)
+						]
+
+					+ SVerticalBox::Slot()
+						.Padding(FMargin(0.0f, 3.0f))
+						.AutoHeight()
+						[
+							SAssignNew(FilterTextBoxWidget, SSearchBox)
+								.HintText(LOCTEXT("FilterSearch", "Search..."))
+								.ToolTipText(LOCTEXT("FilterSearchHint", "Type here to search for a modified object or property."))
+								.OnTextChanged(this, &SUndoHistoryDetails::OnFilterTextChanged)
+						]
+
+					+ SVerticalBox::Slot()
+						.FillHeight(1.0f)
+						.Padding(FMargin(0.0f, 4.0f))
+						[
+							SAssignNew(ChangedObjectsTreeView, STreeView<FUndoDetailsTreeNodePtr>)
+								.TreeItemsSource(&FilteredChangedObjects)
+								.OnGenerateRow(this, &SUndoHistoryDetails::HandleGenerateRow)
+								.OnGetChildren_Static(OnGetChildren)
+								.SelectionMode(ESelectionMode::Single)
+								.HeaderRow
+								(
+									SNew(SHeaderRow)
+									+ SHeaderRow::Column(TransactionDetailsUI::NameLabel)
+									.FillWidth(40.0f)
+									.DefaultLabel(LOCTEXT("NameColumnHeaderName", "Modified objects and properties"))
+									+ SHeaderRow::Column(TransactionDetailsUI::TypeLabel)
+									.FillWidth(20.0f)
+									.DefaultLabel(LOCTEXT("TypeColumnHeaderName", "Type"))
+									+ SHeaderRow::Column(TransactionDetailsUI::ChangeFlagsLabel)
+									.FillWidth(15.0f)
+									.DefaultLabel(LOCTEXT("ChangeFlagsHeaderName", "Change Flags"))
+								)
+						]
+				]
+		];
+}
+
+void SUndoHistoryDetails::Tick(const FGeometry& AllottedGeometry, const double InCurrentTime, const float InDeltaTime) 
+{
+	if (bNeedsRefresh)
+	{
+		Populate();
+	}
+}
+
+void SUndoHistoryDetails::SetSelectedTransaction(const FTransactionDiff& InTransactionDiff)
+{
+	ChangedObjects.Reset();
+
+	TransactionName = FText::FromString(InTransactionDiff.TransactionTitle);
+	TransactionId = FText::FromString(InTransactionDiff.TransactionId.ToString());
+
+	for (const auto& ObjectIt : InTransactionDiff.DiffMap)
+	{
+		UClass* ObjectClass = LoadObject<UClass>(nullptr, *ObjectIt.Value->GetOriginalObjectClassPathName().ToString());
+
+		if (ObjectClass)
+		{
+			ChangedObjects.Emplace(CreateTreeNode(ObjectIt.Value->GetOriginalObjectName().ToString(), ObjectClass, ObjectIt.Value));
+		}
+	}
+
+	Populate();
+}
+
+void SUndoHistoryDetails::Reset()
+{
+	ChangedObjects.Reset();
+	ChangedObjectsTreeView->RequestTreeRefresh();
+}
+
+SUndoHistoryDetails::FUndoDetailsTreeNodePtr SUndoHistoryDetails::CreateTreeNode(const FString& InObjectName, const UClass* InObjectClass, const TSharedPtr<FTransactionObjectEvent>& InEvent) const
+{
+	FUndoDetailsTreeNodePtr ObjectNode = MakeShared<FUndoDetailsTreeNode>(InObjectName, InObjectClass->GetName(), FText::FromName(InEvent->GetOriginalObjectPathName()), InEvent);
+
+	TArray<FUndoHistoryUtils::FBasicPropertyInfo> BasicPropertyInfo = FUndoHistoryUtils::GetChangedPropertiesInfo(InObjectClass, InEvent->GetChangedProperties());
+
+	for (const auto& It : BasicPropertyInfo)
+	{
+		ObjectNode->Children.Emplace(MakeShared<FUndoDetailsTreeNode>(It.PropertyName, It.PropertyType, CreateToolTipText(It.PropertyFlags)));
+	}
+
+	return ObjectNode;
+}
+
+FText SUndoHistoryDetails::CreateToolTipText(EPropertyFlags InFlags) const
+{
+	TArray<const TCHAR*> Flags = ParsePropertyFlags(InFlags);
+	FString ToolTipString = "Flags:\n";
+
+	for (int32 Index = 0; Index < Flags.Num(); Index++)
+	{
+		ToolTipString += FString(Flags[Index]);
+
+		if (Index != Flags.Num() - 1)
+		{
+			ToolTipString += "\n";
+		}
+	}
+
+	return FText::FromString(ToolTipString);
+}
+
+void SUndoHistoryDetails::OnFilterTextChanged(const FText& InFilterText)
+{
+	SearchBoxFilter->SetRawFilterText(InFilterText);
+	FilterTextBoxWidget->SetError(SearchBoxFilter->GetFilterErrorText());
+;}
+
+void SUndoHistoryDetails::FullRefresh()
+{
+	bNeedsExpansion = !SearchBoxFilter->GetRawFilterText().IsEmpty();
+	bNeedsRefresh = true;
+}
+
+void SUndoHistoryDetails::PopulateSearchStrings(const FString& InItemName, TArray< FString >& OutSearchStrings) const
+{
+	OutSearchStrings.Add(InItemName);
+}
+
+void SUndoHistoryDetails::Populate()
+{
+	FilteredChangedObjects.Reset();
+
+	for (const auto& ChangedObject : ChangedObjects)
+	{
+		FUndoDetailsTreeNodePtr ObjectNode = MakeShared<FUndoDetailsTreeNode>(*ChangedObject);
+
+ 		ObjectNode->Children.Reset();
+
+		for (const auto& ChangedProperty : ChangedObject->Children)
+		{
+			if (SearchBoxFilter->PassesFilter(ChangedProperty->Name))
+			{
+				ObjectNode->Children.Add(ChangedProperty);
+			}
+		}
+
+		if (ObjectNode->Children.Num() != 0 || SearchBoxFilter->PassesFilter(ObjectNode->Name))
+		{
+			FilteredChangedObjects.Add(ObjectNode);
+			ChangedObjectsTreeView->SetItemExpansion(ObjectNode, bNeedsExpansion);
+		}
+	}
+
+	ChangedObjectsTreeView->RequestTreeRefresh();
+
+	bNeedsRefresh = false;
+}
+
+TSharedRef<ITableRow> SUndoHistoryDetails::HandleGenerateRow(FUndoDetailsTreeNodePtr InNode, const TSharedRef<STableViewBase>& OwnerTable) const
+{
+	if (InNode->TransactionEvent.IsValid())
+	{
+		return SNew(SUndoHistoryDetailsRow, OwnerTable)
+			.Name(InNode->Name)
+			.Type(InNode->Type)
+			.ToolTipText(FText::FromName(InNode->TransactionEvent->GetOriginalObjectPathName()))
+			.FilterText(this, &SUndoHistoryDetails::HandleGetFilterHighlightText)
+			.TransactionEvent(InNode->TransactionEvent);
+	}
+
+	return SNew(SUndoHistoryDetailsRow, OwnerTable)
+		.Name(InNode->Name)
+		.Type(InNode->Type)
+		.FilterText(this, &SUndoHistoryDetails::HandleGetFilterHighlightText)
+		.ToolTipText(InNode->ToolTip);
+}
+
+FText SUndoHistoryDetails::HandleGetFilterHighlightText() const
+{
+	return SearchBoxFilter->GetRawFilterText();
+}
+
+EVisibility SUndoHistoryDetails::HandleDetailsVisibility() const
+{
+	return ChangedObjects.Num() > 0 ? EVisibility::Visible : EVisibility::Hidden;
+}
+
+FText SUndoHistoryDetails::HandleTransactionName() const
+{
+	return TransactionName;
+}
+
+FText SUndoHistoryDetails::HandleTransactionId() const
+{
+	return TransactionId;
+}
+
+void SUndoHistoryDetails::HandleTransactionIdNavigate()
+{
+	FPlatformApplicationMisc::ClipboardCopy(*TransactionId.ToString());
+}
+
+#undef LOCTEXT_NAMESPACE
\ No newline at end of file
diff --git a/Engine/Source/Developer/UndoHistory/Private/Widgets/UndoHistoryUtils.h b/Engine/Source/Developer/UndoHistory/Private/Widgets/UndoHistoryUtils.h
new file mode 100644
--- /dev/null
+++ b/Engine/Source/Developer/UndoHistory/Private/Widgets/UndoHistoryUtils.h
@@ -0,0 +1,61 @@
+// Copyright Epic Games, Inc. All Rights Reserved.
+
+#pragma once
+
+#include "CoreMinimal.h"
+#include "Widgets/DeclarativeSyntaxSupport.h"
+
+class FUndoHistoryUtils
+{
+public:
+
+	struct FBasicPropertyInfo
+	{
+		FString PropertyName;
+		FString PropertyType;
+		EPropertyFlags PropertyFlags;
+
+		FBasicPropertyInfo(FString InPropertyName, FString InPropertyType, EPropertyFlags InPropertyFlags)
+			: PropertyName(MoveTemp(InPropertyName))
+			, PropertyType(MoveTemp(InPropertyType))
+			, PropertyFlags(InPropertyFlags)
+		{}
+	};
+
+	static TArray<FBasicPropertyInfo> GetChangedPropertiesInfo(const UClass* InObjectClass, const TArray<FName>& InChangedProperties)
+	{
+		TArray<FBasicPropertyInfo> Properties;
+		FString ClassName;
+
+		if (!InObjectClass)
+		{
+			return Properties;
+		}
+
+		for (TFieldIterator<FProperty> Property(InObjectClass); Property; ++Property)
+		{
+			if (!InChangedProperties.Contains(Property->GetFName()))
+			{
+				continue;
+			}
+
+			if (CastField<const FObjectProperty>(*Property) || Property->GetClass() == FStructProperty::StaticClass() || Property->GetClass() == FEnumProperty::StaticClass())
+			{
+				Property->GetCPPMacroType(ClassName);
+			}
+			else if (Property->GetClass() == FArrayProperty::StaticClass())
+			{
+				Property->GetCPPMacroType(ClassName);
+				ClassName = FString::Printf(TEXT("TArray<%s>"), *ClassName);
+			}
+			else
+			{
+				ClassName = Property->GetClass()->GetName();
+				ClassName.RemoveFromEnd("Property");
+			}
+
+			Properties.Emplace(Property->GetName(), MoveTemp(ClassName), Property->GetPropertyFlags());
+		}
+		return Properties;
+	}
+};
\ No newline at end of file
diff --git a/Engine/Source/Developer/UndoHistory/Public/IUndoHistoryModule.h b/Engine/Source/Developer/UndoHistory/Public/IUndoHistoryModule.h
new file mode 100644
--- /dev/null
+++ b/Engine/Source/Developer/UndoHistory/Public/IUndoHistoryModule.h
@@ -0,0 +1,34 @@
+// Copyright Epic Games, Inc. All Rights Reserved.
+
+#pragma once
+
+#include "CoreMinimal.h"
+#include "Modules/ModuleInterface.h"
+#include "Modules/ModuleManager.h"
+
+/**
+ * Interface for the UndoHistory module.
+ */
+class IUndoHistoryModule
+	: public IModuleInterface
+{
+public:
+
+	/**
+	 * Gets a reference to the UndoHistory module instance.
+	 *
+	 * @todo gmp: better implementation using dependency injection.
+	 * @return A reference to the module.
+	 */
+	static IUndoHistoryModule& Get( )
+	{
+		return FModuleManager::LoadModuleChecked<IUndoHistoryModule>("UndoHistory");
+	}
+
+public:
+
+	/**
+	 * Virtual destructor.
+	 */
+	virtual ~IUndoHistoryModule( ) { }
+};
diff --git a/Engine/Source/Developer/UndoHistory/Public/Widgets/SUndoHistoryDetails.h b/Engine/Source/Developer/UndoHistory/Public/Widgets/SUndoHistoryDetails.h
new file mode 100644
--- /dev/null
+++ b/Engine/Source/Developer/UndoHistory/Public/Widgets/SUndoHistoryDetails.h
@@ -0,0 +1,130 @@
+// Copyright Epic Games, Inc. All Rights Reserved.
+
+#pragma once
+
+#include "CoreMinimal.h"
+#include "Widgets/DeclarativeSyntaxSupport.h"
+#include "Widgets/SCompoundWidget.h"
+#include "Widgets/Views/STreeView.h"
+#include "Misc/TextFilter.h"
+
+class UNDOHISTORY_API SUndoHistoryDetails : public SCompoundWidget
+{
+
+public:
+	SLATE_BEGIN_ARGS(SUndoHistoryDetails) { }
+	SLATE_END_ARGS()
+
+public:
+
+	/**
+	 * Construct this widget
+	 *
+	 * @param InArgs The declaration data for this widget.
+	 */
+	void Construct(const FArguments& InArgs);
+
+	/**
+	 * Set the transaction to be displayed in the details panel.
+	 */
+	void SetSelectedTransaction(const struct FTransactionDiff& InTransactionDiff);
+
+	/**
+	 * Clear the details panel.
+	 */
+	void Reset();
+
+
+	/** SWidget interface */
+	virtual void Tick(const FGeometry& AllottedGeometry, const double InCurrentTime, const float InDeltaTime) override;
+
+private:
+
+	struct FUndoDetailsTreeNode;
+	using FUndoDetailsTreeNodePtr = TSharedPtr<FUndoDetailsTreeNode>;
+	using FTreeItemTextFilter = TTextFilter<const FString&>;
+
+	/** Tree node representing a changed object and its changed properties as children. */
+	struct FUndoDetailsTreeNode
+	{
+		FString Name;
+		FString Type;
+		FText ToolTip;
+		TSharedPtr<FTransactionObjectEvent> TransactionEvent;
+		TArray<FUndoDetailsTreeNodePtr> Children;
+
+		FUndoDetailsTreeNode(FString InName, FString InType, FText InToolTip, TSharedPtr<FTransactionObjectEvent> InTransactionEvent = nullptr)
+			: Name(MoveTemp(InName))
+			, Type(MoveTemp(InType))
+			, ToolTip(MoveTemp(InToolTip))
+			, TransactionEvent(MoveTemp(InTransactionEvent))
+		{}
+	};
+
+private:
+
+	/** Create a changed object node. */
+	FUndoDetailsTreeNodePtr CreateTreeNode(const FString& InObjectName, const UClass* InObjectClass, const TSharedPtr<FTransactionObjectEvent>& InEvent) const;
+
+	/** Create a tooltip text for a property. */
+	FText CreateToolTipText(EPropertyFlags InFlags) const;
+
+	/** Callback to handle a change in the filter box. */
+	void OnFilterTextChanged(const FText& InFilterText);
+
+	/** Refresh the details tree view. */
+	void FullRefresh();
+
+	/** Populate the search strings for the filter. */
+	void PopulateSearchStrings(const FString&, TArray<FString>& OutSearchStrings) const;
+
+	/** Populate the details tree. */
+	void Populate();
+
+	/** Callback for generating a UndoHistoryDetailsRow */
+	TSharedRef<ITableRow> HandleGenerateRow(FUndoDetailsTreeNodePtr InNode, const TSharedRef<STableViewBase>& OwnerTable) const;
+
+	/** Callback for getting the filter highlight text. */
+	FText HandleGetFilterHighlightText() const;
+
+	/** Callback for getting the details visibility. */
+	EVisibility HandleDetailsVisibility() const;
+
+	/** Callback for getting the transaction name. */
+	FText HandleTransactionName() const;
+
+	/* Callback for getting the transaction id. */
+	FText HandleTransactionId() const;
+
+	/** Callback for handling a click on the transaction id. */
+	void HandleTransactionIdNavigate();
+
+private:
+
+	/** Holds the ChangedObjects TreeView. */
+	TSharedPtr<STreeView<FUndoDetailsTreeNodePtr> > ChangedObjectsTreeView;
+
+	/** Holds the ChangedObjects to be used as an ItemSource to the TreeView. */
+	TArray<FUndoDetailsTreeNodePtr> ChangedObjects;
+
+	/** Holds the ChangedObjects to be displayed. */
+	TArray<FUndoDetailsTreeNodePtr> FilteredChangedObjects;
+
+	/** Holds the search box. */
+	TSharedPtr<class SSearchBox> FilterTextBoxWidget;
+
+	/** Holds the TransactionName. */
+	FText TransactionName;
+
+	/** Holds the TransactionId. */
+	FText TransactionId;
+
+	/** The TextFilter attached to the SearchBox widget of the UndoHistoryDetails panel. */
+	TSharedPtr<FTreeItemTextFilter> SearchBoxFilter;
+
+	/** If the details tree needs to be refreshed. */
+	bool bNeedsRefresh;
+
+	/** If the tree items need to be expanded (ie. When the filter text changes). */
+	bool bNeedsExpansion;
+};
diff --git a/Engine/Source/Developer/UndoHistory/UndoHistory.Build.cs b/Engine/Source/Developer/UndoHistory/UndoHistory.Build.cs
new file mode 100644
--- /dev/null
+++ b/Engine/Source/Developer/UndoHistory/UndoHistory.Build.cs
@@ -0,0 +1,26 @@
+// Copyright Epic Games, Inc. All Rights Reserved.
+
+namespace UnrealBuildTool.Rules
+{
+	public class UndoHistory : ModuleRules
+	{
+		public UndoHistory(ReadOnlyTargetRules Target) : base(Target)
+		{
+			PublicDependencyModuleNames.AddRange(
+				new string[] {
+					"Core",
+				}
+			);
+
+			PrivateDependencyModuleNames.AddRange(
+				new string[] {
+					"ApplicationCore",
+					"CoreUObject",
+					"InputCore",
+					"Slate",
+					"SlateCore",
+				}
+			);
+		}
+	}
+}
diff --git a/Engine/Source/Editor/MainFrame/MainFrame.Build.cs b/Engine/Source/Editor/MainFrame/MainFrame.Build.cs
--- a/Engine/Source/Editor/MainFrame/MainFrame.Build.cs
+++ b/Engine/Source/Editor/MainFrame/MainFrame.Build.cs
@@ -39,7 +39,7 @@ public class MainFrame : ModuleRules
 				"TranslationEditor",
 				"Projects",
 				"DeviceProfileEditor",
-				"UndoHistory",
+				"UndoHistoryEditor",
 				"Analytics",
 				"ToolMenus",
 				"LauncherServices",
diff --git a/Engine/Source/Editor/MainFrame/Private/Menus/MainMenu.cpp b/Engine/Source/Editor/MainFrame/Private/Menus/MainMenu.cpp
--- a/Engine/Source/Editor/MainFrame/Private/Menus/MainMenu.cpp
+++ b/Engine/Source/Editor/MainFrame/Private/Menus/MainMenu.cpp
@@ -7,6 +7,7 @@
 #include "Interfaces/IMainFrameModule.h"
 #include "DesktopPlatformModule.h"
 #include "ISourceControlModule.h"
+#include "IUndoHistoryEditorModule.h"
 #include "Toolkits/AssetEditorToolkit.h"
 #include "Toolkits/GlobalEditorCommonCommands.h"
 #include "SourceCodeNavigation.h"
@@ -26,7 +27,6 @@
 #include "WorkspaceMenuStructureModule.h"
 #include "Features/EditorFeatures.h"
 #include "Features/IModularFeatures.h"
-#include "UndoHistoryModule.h"
 #include "Framework/Commands/GenericCommands.h"
 #include "ITranslationEditor.h"
 #include "LauncherPlatformModule.h"
@@ -101,7 +101,7 @@ void FMainMenu::RegisterEditMenu()
 			LOCTEXT("UndoHistoryTabTitle", "Undo History"),
 			LOCTEXT("UndoHistoryTooltipText", "View the entire undo history."),
 			FSlateIcon(FEditorStyle::GetStyleSetName(), "UndoHistory.TabIcon"),
-			FUIAction(FExecuteAction::CreateStatic(&FUndoHistoryModule::ExecuteOpenUndoHistory))
+			FUIAction(FExecuteAction::CreateLambda([](){  IUndoHistoryEditorModule::Get().ExecuteOpenUndoHistory(); } ))
 			);
 	}
 
diff --git a/Engine/Source/Editor/UndoHistory/Private/Classes/UndoHistorySettings.h b/Engine/Source/Editor/UndoHistory/Private/Classes/UndoHistorySettings.h
deleted file mode 100644
diff --git a/Engine/Source/Editor/UndoHistory/Private/UndoHistoryModule.cpp b/Engine/Source/Editor/UndoHistory/Private/UndoHistoryModule.cpp
deleted file mode 100644
diff --git a/Engine/Source/Editor/UndoHistory/Private/Widgets/SUndoHistory.cpp b/Engine/Source/Editor/UndoHistory/Private/Widgets/SUndoHistory.cpp
deleted file mode 100644
diff --git a/Engine/Source/Editor/UndoHistory/Private/Widgets/SUndoHistory.h b/Engine/Source/Editor/UndoHistory/Private/Widgets/SUndoHistory.h
deleted file mode 100644
diff --git a/Engine/Source/Editor/UndoHistory/Private/Widgets/SUndoHistoryDetails.cpp b/Engine/Source/Editor/UndoHistory/Private/Widgets/SUndoHistoryDetails.cpp
deleted file mode 100644
diff --git a/Engine/Source/Editor/UndoHistory/Private/Widgets/SUndoHistoryTableRow.h b/Engine/Source/Editor/UndoHistory/Private/Widgets/SUndoHistoryTableRow.h
deleted file mode 100644
diff --git a/Engine/Source/Editor/UndoHistory/Private/Widgets/UndoHistoryUtils.h b/Engine/Source/Editor/UndoHistory/Private/Widgets/UndoHistoryUtils.h
deleted file mode 100644
diff --git a/Engine/Source/Editor/UndoHistory/Public/Interfaces/IUndoHistoryModule.h b/Engine/Source/Editor/UndoHistory/Public/Interfaces/IUndoHistoryModule.h
deleted file mode 100644
diff --git a/Engine/Source/Editor/UndoHistory/Public/UndoHistory.h b/Engine/Source/Editor/UndoHistory/Public/UndoHistory.h
deleted file mode 100644
diff --git a/Engine/Source/Editor/UndoHistory/Public/UndoHistoryModule.h b/Engine/Source/Editor/UndoHistory/Public/UndoHistoryModule.h
deleted file mode 100644
diff --git a/Engine/Source/Editor/UndoHistory/Public/Widgets/SUndoHistoryDetails.h b/Engine/Source/Editor/UndoHistory/Public/Widgets/SUndoHistoryDetails.h
deleted file mode 100644
diff --git a/Engine/Source/Editor/UndoHistory/UndoHistory.Build.cs b/Engine/Source/Editor/UndoHistory/UndoHistory.Build.cs
deleted file mode 100644
diff --git a/Engine/Source/Editor/UndoHistoryEditor/Private/UndoHistoryEditorModule.cpp b/Engine/Source/Editor/UndoHistoryEditor/Private/UndoHistoryEditorModule.cpp
new file mode 100644
--- /dev/null
+++ b/Engine/Source/Editor/UndoHistoryEditor/Private/UndoHistoryEditorModule.cpp
@@ -0,0 +1,54 @@
+// Copyright Epic Games, Inc. All Rights Reserved.
+
+#include "IUndoHistoryEditorModule.h"
+#include "EditorStyleSet.h"
+#include "Modules/ModuleManager.h"
+#include "Widgets/SUndoHistory.h"
+#include "Widgets/Docking/SDockTab.h"
+
+const FName UndoHistoryTabName("UndoHistory");
+
+class FUndoHistoryEditorModule : public IUndoHistoryEditorModule
+{
+public:
+
+	virtual void StartupModule() override
+	{
+		FGlobalTabmanager::Get()->RegisterNomadTabSpawner(UndoHistoryTabName, FOnSpawnTab::CreateRaw(this, &FUndoHistoryEditorModule::HandleSpawnSettingsTab))
+			.SetDisplayName(NSLOCTEXT("FUndoHistoryModule", "UndoHistoryTabTitle", "Undo History"))
+			.SetTooltipText(NSLOCTEXT("FUndoHistoryModule", "UndoHistoryTooltipText", "Open the Undo History tab."))
+			.SetIcon(FSlateIcon(FEditorStyle::GetStyleSetName(), "UndoHistory.TabIcon"))
+			.SetAutoGenerateMenuEntry(false);
+	}
+	
+	virtual void ShutdownModule() override
+	{
+		FGlobalTabmanager::Get()->UnregisterNomadTabSpawner(UndoHistoryTabName);
+	}
+	
+	virtual bool SupportsDynamicReloading() override
+	{
+		return true;
+	}
+
+	virtual void ExecuteOpenUndoHistory() override
+	{
+		FGlobalTabmanager::Get()->TryInvokeTab(UndoHistoryTabName);
+	}
+
+private:
+	
+	/** Handles creating the project settings tab. */
+	TSharedRef<SDockTab> HandleSpawnSettingsTab(const FSpawnTabArgs& SpawnTabArgs)
+	{
+		const TSharedRef<SDockTab> DockTab = SNew(SDockTab)
+			.TabRole(ETabRole::NomadTab)
+		[
+			SNew(SUndoHistory)
+		];
+		return DockTab;
+	}
+};
+
+IMPLEMENT_MODULE(FUndoHistoryEditorModule, UndoHistoryEditor);
+
diff --git a/Engine/Source/Editor/UndoHistoryEditor/Private/UndoHistorySettings.h b/Engine/Source/Editor/UndoHistoryEditor/Private/UndoHistorySettings.h
new file mode 100644
--- /dev/null
+++ b/Engine/Source/Editor/UndoHistoryEditor/Private/UndoHistorySettings.h
@@ -0,0 +1,23 @@
+// Copyright Epic Games, Inc. All Rights Reserved.
+
+#pragma once
+#include "CoreMinimal.h"
+#include "UObject/ObjectMacros.h"
+#include "UObject/Object.h"
+#include "UndoHistorySettings.generated.h"
+
+
+/**
+ * Implements the settings for the UndoHistory.
+ */
+UCLASS(config = EditorPerProjectUserSettings)
+class UUndoHistorySettings
+	: public UObject
+{
+public:
+	GENERATED_BODY()
+
+	/** True when the UndoHistory is showing transaction details. */
+	UPROPERTY(config)
+	bool bShowTransactionDetails;
+};
diff --git a/Engine/Source/Editor/UndoHistoryEditor/Private/Widgets/SUndoHistory.cpp b/Engine/Source/Editor/UndoHistoryEditor/Private/Widgets/SUndoHistory.cpp
new file mode 100644
--- /dev/null
+++ b/Engine/Source/Editor/UndoHistoryEditor/Private/Widgets/SUndoHistory.cpp
@@ -0,0 +1,497 @@
+// Copyright Epic Games, Inc. All Rights Reserved.
+
+#include "Widgets/SUndoHistory.h"
+#include "SlateOptMacros.h"
+#include "Widgets/Images/SImage.h"
+#include "Widgets/Text/STextBlock.h"
+#include "Widgets/Input/SButton.h"
+#include "EditorStyleSet.h"
+#include "Editor.h"
+#include "Framework/MultiBox/MultiBoxBuilder.h"
+#include "Widgets/Layout/SExpandableArea.h"
+#include "Widgets/Views/SListView.h"
+#include "Widgets/SUndoHistoryTableRow.h"
+#include "Widgets/Input/SComboButton.h"
+#include "Editor/TransBuffer.h"
+#include "UndoHistorySettings.h"
+#include "Misc/ScopedSlowTask.h"
+#include "Widgets/SUndoHistoryDetails.h"
+
+#define LOCTEXT_NAMESPACE "SUndoHistory"
+
+
+// Helper function for TransactionInfo
+const FTransaction* SUndoHistory::FTransactionInfo::GetTransaction() const
+{
+	if (GEditor && GEditor->Trans)
+	{
+		return GEditor->Trans->GetTransaction(QueueIndex);
+	}
+	return nullptr;
+}
+
+/* SUndoHistory interface
+ *****************************************************************************/
+
+BEGIN_SLATE_FUNCTION_BUILD_OPTIMIZATION
+void SUndoHistory::Construct( const FArguments& InArgs )
+{
+	LastActiveTransactionIndex = INDEX_NONE;
+
+	ChildSlot
+	[
+		SNew(SVerticalBox)
+
+		+ SVerticalBox::Slot()
+			.FillHeight(1.0f)
+			[
+				SAssignNew(Splitter, SSplitter)
+					.Orientation(EOrientation::Orient_Vertical)
+
+					+ SSplitter::Slot()
+						[
+							SNew(SVerticalBox)
+								+ SVerticalBox::Slot()
+								.Padding(FMargin(2.0f))
+								.AutoHeight()
+								[
+									SNew(SHorizontalBox)
+
+									+SHorizontalBox::Slot()
+										.AutoWidth()
+										.Padding( 0.0f, 0.0f, 4.0f, 0.0f )
+										[
+											SNew( SImage )
+												.Image( FEditorStyle::GetBrush( "LevelEditor.Tabs.Details" ) )
+										]
+
+									+SHorizontalBox::Slot()
+										.HAlign(HAlign_Left)
+										[
+											SNew(STextBlock)
+												.Text(LOCTEXT("TransactionHistoryLabel", "Transactions"))
+										]
+								]
+							
+								+ SVerticalBox::Slot()
+								[
+									SNew(SBorder)
+										.BorderImage(FEditorStyle::GetBrush("ToolPanel.GroupBorder"))
+										.Padding(FMargin(4.0f, 1.0f))
+										[
+											SAssignNew(UndoListView, SListView<TSharedPtr<FTransactionInfo> >)
+												.ItemHeight(24.0f)
+												.ListItemsSource(&UndoList)
+												.SelectionMode(ESelectionMode::Single)
+												.OnGenerateRow(this, &SUndoHistory::HandleUndoListGenerateRow)
+												.OnMouseButtonDoubleClick(this, &SUndoHistory::HandleUndoListJumpToTransaction)
+												.OnSelectionChanged(this, &SUndoHistory::HandleUndoListSelectionChanged)
+												.HeaderRow
+												(
+													SNew(SHeaderRow)
+														.Visibility(EVisibility::Collapsed)
+
+													+ SHeaderRow::Column("Title")
+														.FillWidth(80.0f)
+
+													+ SHeaderRow::Column("JumpToButton")
+														.FixedWidth(30.0f)
+												)
+										]
+								]
+
+						]
+
+					+ SSplitter::Slot()
+						.Value(0.4f)
+						[
+							SNew(SVerticalBox)
+							.Visibility(this, &SUndoHistory::HandleUndoHistoryDetailsVisibility)
+							+SVerticalBox::Slot()
+							.Padding(FMargin(2.0f))
+							.AutoHeight()
+							[
+								SNew(SHorizontalBox)
+
+								+ SHorizontalBox::Slot()
+									.AutoWidth()
+									.Padding( 0.0f, 0.0f, 4.0f, 0.0f )
+									[
+										SNew(SImage)
+											.Image(FEditorStyle::GetBrush( "LevelEditor.Tabs.Details" ))
+									]
+
+								+ SHorizontalBox::Slot()
+									.Padding(FMargin(4.0f, 0.0f))
+									[
+										SNew(STextBlock)
+											.TextStyle(FEditorStyle::Get(), "NormalText")
+											.Text(LOCTEXT("TransactionDetailsLabel", "Transaction Details"))
+									]
+							]
+							
+							+SVerticalBox::Slot()
+							[
+								SAssignNew(UndoDetailsView, SUndoHistoryDetails)
+							]
+						]
+			]
+
+		+ SVerticalBox::Slot()
+			.AutoHeight()
+			.Padding(0.0f, 2.0f, 0.0f, 0.0f)
+			[
+				SNew(SHorizontalBox)
+
+				+ SHorizontalBox::Slot()
+					.AutoWidth()
+					.Padding(2.0f, 0.0f, 0.0f, 0.0f)
+					.VAlign(VAlign_Center)
+					[
+						SNew(STextBlock)
+							.TextStyle(FEditorStyle::Get(), "ContentBrowser.TopBar.Font")
+							.Font(FEditorStyle::Get().GetFontStyle("FontAwesome.11"))
+							.ToolTipText(LOCTEXT("UndoBufferFullToolTip", "The undo buffer has reached its maximum capacity, transactions will be deleted from the top."))
+							.Text(FText::FromString(FString(TEXT("\xf071 "))) /*fa-exclamation-triangle*/)
+							.ColorAndOpacity(FEditorStyle::Get().GetWidgetStyle<FButtonStyle>("FlatButton.Danger").Normal.TintColor)
+							.Visibility(this, &SUndoHistory::HandleUndoWarningVisibility)
+					]
+
+				+ SHorizontalBox::Slot()
+					.AutoWidth()
+					.VAlign(VAlign_Center)
+					[
+						// Undo size
+						SNew(STextBlock)
+							.Text(this, &SUndoHistory::HandleUndoSizeTextBlockText)
+					]
+
+				+ SHorizontalBox::Slot() 
+					.Padding(8.0f, 0.0f, 0.0f, 0.0f)
+					.AutoWidth()
+					.HAlign(HAlign_Left)
+					[
+						// Discard history button
+						SAssignNew(DiscardButton, SButton)
+							.ForegroundColor(FSlateColor::UseForeground())
+							.ButtonStyle(FEditorStyle::Get(), "ToggleButton")
+ 							.OnClicked(this, &SUndoHistory::HandleDiscardHistoryButtonClicked)
+							.ToolTipText(LOCTEXT("DiscardHistoryButtonToolTip", "Discard the Undo History."))
+							[
+								SNew(SImage)
+								.Image(FEditorStyle::Get().GetBrush("Icons.Delete"))
+							]
+					]
+
+				+ SHorizontalBox::Slot()
+					.HAlign(HAlign_Right)
+					[
+						SNew(SComboButton)
+							.ContentPadding(0)
+							.ForegroundColor(FSlateColor::UseForeground())
+							.ButtonStyle(FEditorStyle::Get(), "ToggleButton")
+							.OnGetMenuContent(this, &SUndoHistory::GetViewButtonContent)
+							.ButtonContent()
+							[
+								SNew(SHorizontalBox)
+
+								+ SHorizontalBox::Slot()
+									.AutoWidth()
+									.VAlign(VAlign_Center)
+									[
+										SNew(SImage).Image(FEditorStyle::GetBrush("GenericViewButton"))
+									]
+							]
+					]
+			]
+	];
+
+	ReloadUndoList();
+
+	if (GEditor != nullptr && GEditor->Trans != nullptr)
+	{
+		UTransBuffer* TransBuffer = CastChecked<UTransBuffer>(GEditor->Trans);
+		TransBuffer->OnUndoBufferChanged().AddRaw(this, &SUndoHistory::OnUndoBufferChanged);
+		TransBuffer->OnTransactionStateChanged().AddRaw(this, &SUndoHistory::OnTransactionStateChanged);
+	}
+}
+END_SLATE_FUNCTION_BUILD_OPTIMIZATION
+
+SUndoHistory::~SUndoHistory()
+{
+	if (GEditor != nullptr && GEditor->Trans != nullptr)
+	{
+		UTransBuffer* TransBuffer = CastChecked<UTransBuffer>(GEditor->Trans);
+		TransBuffer->OnUndoBufferChanged().RemoveAll(this);
+		TransBuffer->OnTransactionStateChanged().RemoveAll(this);
+	}
+}
+
+void SUndoHistory::OnUndoBufferChanged()
+{
+	ReloadUndoList();
+
+	SelectLastTransaction();
+}
+
+
+void SUndoHistory::OnTransactionStateChanged(const FTransactionContext& TransactionContext, ETransactionStateEventType TransactionState)
+{
+	if (TransactionState == ETransactionStateEventType::TransactionFinalized)
+	{
+		SelectLastTransaction();
+	}
+}
+
+void SUndoHistory::Tick(const FGeometry& AllottedGeometry, const double InCurrentTime, const float InDeltaTime)
+{
+	// reload the transaction list if necessary
+	if (GEditor == nullptr || GEditor->Trans == nullptr)
+	{
+		ReloadUndoList();
+	}
+	else if (GEditor != nullptr && GEditor->Trans != nullptr)
+	{
+		// update the selected transaction if necessary
+		int32 ActiveTransactionIndex = GEditor->Trans->GetQueueLength() - GEditor->Trans->GetUndoCount() - 1;
+
+		if (UndoList.IsValidIndex(ActiveTransactionIndex) && (ActiveTransactionIndex != LastActiveTransactionIndex))
+		{
+			UndoListView->SetSelection(UndoList[ActiveTransactionIndex]);
+			LastActiveTransactionIndex = ActiveTransactionIndex;
+		}
+	}
+}
+
+
+/* SUndoHistory implementation
+ *****************************************************************************/
+
+void SUndoHistory::ReloadUndoList()
+{
+	UndoList.Empty();
+
+	if ((GEditor == nullptr) || (GEditor->Trans == nullptr))
+	{
+		return;
+	}
+
+	for (int32 QueueIndex = 0; QueueIndex < GEditor->Trans->GetQueueLength(); ++QueueIndex)
+	{
+		UndoList.Add(MakeShared<FTransactionInfo>(QueueIndex));
+	}
+
+	UndoListView->RequestListRefresh();
+}
+
+void SUndoHistory::SelectLastTransaction()
+{
+	if (GEditor != nullptr && GEditor->Trans != nullptr)
+	{
+		LastActiveTransactionIndex = GEditor->Trans->GetQueueLength() - 1;
+		if (UndoList.IsValidIndex(LastActiveTransactionIndex))
+		{
+			const TSharedPtr<FTransactionInfo>& TransactionInfo = UndoList[LastActiveTransactionIndex];
+			const FTransaction* InnerTransaction = TransactionInfo.IsValid() ? TransactionInfo->GetTransaction() : nullptr;
+			if (InnerTransaction)
+			{
+				UndoListView->SetSelection(TransactionInfo);
+				UndoDetailsView->SetSelectedTransaction(InnerTransaction->GenerateDiff());
+			}
+		}
+	}
+}
+
+
+/* SUndoHistory callbacks
+ *****************************************************************************/
+
+FReply SUndoHistory::HandleDiscardHistoryButtonClicked()
+{
+	if ((GEditor != nullptr) && (GEditor->Trans != nullptr))
+	{
+		GEditor->Trans->Reset(LOCTEXT("DiscardHistoryReason", "Discard undo history."));
+		ReloadUndoList();
+		UndoDetailsView->Reset();
+	}
+
+	return FReply::Handled();
+}
+
+
+TSharedRef<ITableRow> SUndoHistory::HandleUndoListGenerateRow(TSharedPtr<FTransactionInfo> TransactionInfo, const TSharedRef<STableViewBase>& OwnerTable)
+{
+	return SNew(SUndoHistoryTableRow, OwnerTable)
+		.OnGotoTransactionClicked(this, &SUndoHistory::HandleGoToTransaction)
+		.IsApplied(this, &SUndoHistory::HandleUndoListRowIsApplied, TransactionInfo->QueueIndex)
+		.QueueIndex(TransactionInfo->QueueIndex)
+		.Transaction(TransactionInfo->GetTransaction());
+}
+
+void SUndoHistory::HandleGoToTransaction(const FGuid& TargetTransactionId)
+{
+	TSharedPtr<FTransactionInfo>* Transaction = UndoList.FindByPredicate([TargetTransactionId](TSharedPtr<FTransactionInfo> TransactionInfo)
+	{
+		if (const FTransaction* InnerTransaction = TransactionInfo->GetTransaction())
+		{
+			return InnerTransaction->GetId() == TargetTransactionId;
+		}
+		return false;
+	});
+
+	if (Transaction != nullptr)
+	{
+		HandleUndoListJumpToTransaction(*Transaction);
+	}
+}
+
+bool SUndoHistory::HandleUndoListRowIsApplied(int32 QueueIndex) const
+{
+	if ((GEditor == nullptr) || (GEditor->Trans == nullptr))
+	{
+		return false;
+	}
+
+	return (QueueIndex < (GEditor->Trans->GetQueueLength() - GEditor->Trans->GetUndoCount()));
+}
+
+
+void SUndoHistory::HandleUndoListSelectionChanged(TSharedPtr<FTransactionInfo> InItem, ESelectInfo::Type SelectInfo)
+{
+	if (!InItem.IsValid() || (GEditor == nullptr) || (GEditor->Trans == nullptr))
+	{
+		return;
+	}
+
+	if (const FTransaction* InnerTransaction = InItem->GetTransaction())
+	{
+		UndoDetailsView->SetSelectedTransaction(InnerTransaction->GenerateDiff());
+	}
+
+	if (SelectInfo == ESelectInfo::OnMouseClick || SelectInfo == ESelectInfo::Direct)
+	{
+		// Use Private_SetItemSelection to avoid creating an OnSelectionChanged broadcast
+		UndoListView->Private_SetItemSelection(InItem, true, false);
+		UndoListView->RequestScrollIntoView(InItem);
+	}
+
+
+}
+
+void SUndoHistory::HandleUndoListJumpToTransaction(TSharedPtr<FTransactionInfo> InItem)
+{
+	if (!InItem.IsValid() || (GEditor == nullptr) || (GEditor->Trans == nullptr))
+	{
+		return;
+	}
+
+	LastActiveTransactionIndex = GEditor->Trans->GetQueueLength() - GEditor->Trans->GetUndoCount() - 1;
+
+	const int32 RemainingUndoRedo = FMath::Abs(InItem->QueueIndex - LastActiveTransactionIndex);
+
+ 	FScopedSlowTask SlowTask(RemainingUndoRedo, LOCTEXT("ReplayingTransactions", "Replaying Transactions..."));
+ 	SlowTask.MakeDialogDelayed(0.5f);
+
+	const auto ProgressLambda = [&SlowTask](const int32 InTransactionIndex)
+	{
+		const FTransaction* Transaction = GEditor->Trans->GetTransaction(InTransactionIndex);
+
+		if (Transaction)
+		{
+			SlowTask.EnterProgressFrame(1.0f, FText::Format(LOCTEXT("ReplayingTransactionFmt", "Replaying {0}"), Transaction->GetTitle()));
+		}
+	};
+
+	while (InItem->QueueIndex > LastActiveTransactionIndex)
+	{
+		if (!GEditor->Trans->Redo())
+		{	
+			break;
+		}
+
+		ProgressLambda(LastActiveTransactionIndex);
+		++LastActiveTransactionIndex;
+	}
+
+	while (InItem->QueueIndex < LastActiveTransactionIndex)
+	{
+		if (!GEditor->Trans->Undo())
+		{
+			break;
+		}
+
+		ProgressLambda(LastActiveTransactionIndex);
+		--LastActiveTransactionIndex;
+	}
+
+	UndoListView->RequestScrollIntoView(InItem);
+	UndoListView->SetSelection(InItem);
+}
+
+
+EVisibility SUndoHistory::HandleUndoHistoryDetailsVisibility() const
+{
+	return IsShowingTransactionDetails() ? EVisibility::Visible : EVisibility::Collapsed;
+}
+
+FText SUndoHistory::HandleUndoSizeTextBlockText() const
+{
+	if ((GEditor == nullptr) || (GEditor->Trans == nullptr))
+	{
+		return FText::GetEmpty();
+	}
+
+	return FText::Format(LOCTEXT("TransactionCountF", "{0} Transactions ({1})"), FText::AsNumber(UndoList.Num()), FText::AsMemory(GEditor->Trans->GetUndoSize()));
+}
+
+EVisibility SUndoHistory::HandleUndoWarningVisibility() const
+{
+	if ((GEditor == nullptr) || (GEditor->Trans == nullptr))
+	{
+		return EVisibility::Collapsed;
+	}
+
+	UTransBuffer* TransBuffer = CastChecked<UTransBuffer>(GEditor->Trans);
+
+	return TransBuffer->GetUndoSize() > TransBuffer->MaxMemory ? EVisibility::Visible : EVisibility::Collapsed;
+}
+
+TSharedRef<SWidget> SUndoHistory::GetViewButtonContent() const
+{
+	FMenuBuilder MenuBuilder(true, NULL);
+
+	MenuBuilder.BeginSection("AssetThumbnails", LOCTEXT("ShowHeading", "Show"));
+	{
+		MenuBuilder.AddMenuEntry(
+			LOCTEXT("ToggleShowTransactionDetails", "Show transactions details."),
+			LOCTEXT("ToggleShowTransactionDetailsToolTip", "When enabled, display additional information about transactions."),
+			FSlateIcon(),
+			FUIAction(
+				FExecuteAction::CreateSP(const_cast<SUndoHistory*>(this), &SUndoHistory::ToggleShowTransactionDetails),
+				FCanExecuteAction(),
+				FIsActionChecked::CreateSP(this, &SUndoHistory::IsShowingTransactionDetails)
+			),
+			NAME_None,
+			EUserInterfaceActionType::ToggleButton
+		);
+	}
+	MenuBuilder.EndSection();
+
+	return MenuBuilder.MakeWidget();
+}
+
+/** FILTERS */
+void SUndoHistory::ToggleShowTransactionDetails()
+{
+	const bool bEnableFlag = !IsShowingTransactionDetails();
+
+	UUndoHistorySettings* Settings = GetMutableDefault<UUndoHistorySettings>();
+	Settings->bShowTransactionDetails = bEnableFlag;
+}
+
+bool SUndoHistory::IsShowingTransactionDetails() const
+{
+	return GetDefault<UUndoHistorySettings>()->bShowTransactionDetails;
+}
+
+#undef LOCTEXT_NAMESPACE
\ No newline at end of file
diff --git a/Engine/Source/Editor/UndoHistoryEditor/Private/Widgets/SUndoHistory.h b/Engine/Source/Editor/UndoHistoryEditor/Private/Widgets/SUndoHistory.h
new file mode 100644
--- /dev/null
+++ b/Engine/Source/Editor/UndoHistoryEditor/Private/Widgets/SUndoHistory.h
@@ -0,0 +1,125 @@
+// Copyright Epic Games, Inc. All Rights Reserved.
+
+#pragma once
+
+#include "CoreMinimal.h"
+#include "Input/Reply.h"
+#include "Widgets/DeclarativeSyntaxSupport.h"
+#include "Widgets/SCompoundWidget.h"
+#include "Widgets/Views/STableViewBase.h"
+#include "Widgets/Views/STableRow.h"
+#include "Misc/ITransaction.h"
+
+class FTransaction;
+
+/**
+ * Implements the undo history panel.
+ */
+class SUndoHistory
+	: public SCompoundWidget
+{
+	// Structure for transaction information.
+	struct FTransactionInfo
+	{
+		// Holds the transactions index in the transaction queue.
+		int32 QueueIndex;
+
+		// Creates and initializes a new instance.
+		FTransactionInfo( int32 InQueueIndex )
+			: QueueIndex(InQueueIndex)			
+		{ }
+
+		const FTransaction* GetTransaction() const;
+	};
+
+public:
+
+	SLATE_BEGIN_ARGS(SUndoHistory) { }
+	SLATE_END_ARGS()
+
+	virtual ~SUndoHistory();
+
+	/**
+	 * Construct this widget
+	 *
+	 * @param InArgs The declaration data for this widget.
+	 */
+	void Construct( const FArguments& InArgs );
+	
+	//~ Begin SWidget Interface
+	virtual void Tick( const FGeometry& AllottedGeometry, const double InCurrentTime, const float InDeltaTime ) override;
+	//~ End SWidget Interface
+
+protected:
+
+	/**
+	 * Reloads the list of undo transactions.
+	 */
+	void ReloadUndoList( );
+	void OnUndoBufferChanged();
+	void OnTransactionStateChanged(const FTransactionContext& TransactionContext, ETransactionStateEventType TransactionState);
+
+private:
+
+	/** Select the last transaction in the undo history. */
+	void SelectLastTransaction();
+
+	/** Callback for clicking the 'Discard History' button. */
+	FReply HandleDiscardHistoryButtonClicked( );
+
+	/** Callback for generating a row widget for the message list view. */
+	TSharedRef<ITableRow> HandleUndoListGenerateRow( TSharedPtr<FTransactionInfo> TransactionInfo, const TSharedRef<STableViewBase>& OwnerTable );
+
+	/** Callback for when a user wants to jump to a certain transaction. */ 
+	void HandleGoToTransaction(const FGuid& TargetTransactionId);
+
+	/** Callback for checking whether the specified undo list row transaction is applied. */
+	bool HandleUndoListRowIsApplied( int32 QueueIndex ) const;
+
+	/** Callback for selecting a message in the message list view. */
+	void HandleUndoListSelectionChanged( TSharedPtr<FTransactionInfo> TransactionInfo, ESelectInfo::Type SelectInfo );
+
+	/** Callback for getting the undo size text. */
+	FText HandleUndoSizeTextBlockText( ) const;
+
+	/** Callback for getting the right undo size text color. */
+	EVisibility HandleUndoWarningVisibility() const;
+
+	/** Callback for clicking on a transaction to display details. */
+	void HandleUndoListJumpToTransaction(TSharedPtr<FTransactionInfo> InItem);
+
+	/** Callback to handle UndoHistoryDetails visibility. */
+	EVisibility HandleUndoHistoryDetailsVisibility() const;
+
+	/** Callback for getting the view button's content. */
+	TSharedRef<SWidget> GetViewButtonContent() const;
+
+	/** Toggle visibility of the transaction details section. */
+	void ToggleShowTransactionDetails();
+
+	/** Whether the transaction details section should be displayed. */
+	bool IsShowingTransactionDetails() const;
+
+private:
+
+	/** Holds the index of the last active transaction. */
+	int32 LastActiveTransactionIndex;
+
+	/** Holds the list of undo transaction indices. */
+	TArray<TSharedPtr<FTransactionInfo> > UndoList;
+
+	/** Holds the undo list view. */
+	TSharedPtr<SListView<TSharedPtr<FTransactionInfo> > > UndoListView;
+
+	/** Holds the undo details panel view. */
+	TSharedPtr<class SUndoHistoryDetails> UndoDetailsView;
+
+	/** Holds the Transaction panel splitter. */
+	TSharedPtr<SSplitter> Splitter;
+
+	/** Holds the Details slot from the splitter. */
+	TSharedPtr<SSplitter::FSlot> DetailsSlot;
+
+	/** Holds the undo history discard button. */
+	TSharedPtr<SButton> DiscardButton;
+};
diff --git a/Engine/Source/Editor/UndoHistoryEditor/Private/Widgets/SUndoHistoryTableRow.h b/Engine/Source/Editor/UndoHistoryEditor/Private/Widgets/SUndoHistoryTableRow.h
new file mode 100644
--- /dev/null
+++ b/Engine/Source/Editor/UndoHistoryEditor/Private/Widgets/SUndoHistoryTableRow.h
@@ -0,0 +1,155 @@
+// Copyright Epic Games, Inc. All Rights Reserved.
+
+#pragma once
+
+#include "CoreMinimal.h"
+#include "Misc/Attribute.h"
+#include "Widgets/SNullWidget.h"
+#include "Widgets/DeclarativeSyntaxSupport.h"
+#include "Styling/SlateColor.h"
+#include "Widgets/SWidget.h"
+#include "Layout/Margin.h"
+#include "Widgets/Views/STableViewBase.h"
+#include "Widgets/Layout/SBox.h"
+#include "Widgets/Views/STableRow.h"
+#include "SlateOptMacros.h"
+#include "Widgets/Text/STextBlock.h"
+#include "Editor/Transactor.h"
+#include "Widgets/Views/SListView.h"
+#include "Widgets/Views/SExpanderArrow.h"
+#include "Widgets/Layout/SExpandableArea.h"
+#include "EditorStyleSet.h"
+#include "Widgets/Views/STreeView.h"
+#include "Widgets/Input/SButton.h"
+
+#define LOCTEXT_NAMESPACE "SUndoHistoryTableRow"
+
+DECLARE_DELEGATE_OneParam(FOnGotoTransactionClicked, const FGuid&)
+
+/**
+ * Implements a row widget for the undo history list.
+ */
+class SUndoHistoryTableRow
+	: public SMultiColumnTableRow<TSharedPtr<int32> >
+{
+
+public:
+
+	SLATE_BEGIN_ARGS(SUndoHistoryTableRow) { }
+		SLATE_ATTRIBUTE(bool, IsApplied)
+		SLATE_ARGUMENT(int32, QueueIndex)
+		SLATE_ARGUMENT(const FTransaction*, Transaction)
+		SLATE_EVENT(FOnGotoTransactionClicked, OnGotoTransactionClicked)
+	SLATE_END_ARGS()
+
+public:
+
+	/**
+	 * Constructs the widget.
+	 *
+	 * @param InArgs The construction arguments.
+	 */
+	void Construct( const FArguments& InArgs, const TSharedRef<STableViewBase>& InOwnerTableView )
+	{
+		IsApplied = InArgs._IsApplied;
+		QueueIndex = InArgs._QueueIndex;
+		OnGotoTransactionClicked = InArgs._OnGotoTransactionClicked;
+
+		FSuperRowType::FArguments Args = FSuperRowType::FArguments()
+			.Style(&FEditorStyle::Get().GetWidgetStyle<FTableRowStyle>("SceneOutliner.TableViewRow"));
+
+		UObject* ContextObject = nullptr;
+		if (InArgs._Transaction)
+		{
+			TransactionId = InArgs._Transaction->GetId();
+			ContextObject = InArgs._Transaction->GetContext().PrimaryObject;
+			Title = InArgs._Transaction->GetTitle();
+			Description = InArgs._Transaction->GetDescription();
+		}
+
+		if (ContextObject != nullptr)
+		{
+			Title = FText::Format(LOCTEXT("UndoHistoryTableRowTitleF", "{0} [{1}]"), InArgs._Transaction->GetTitle(), FText::FromString(ContextObject->GetFName().ToString()));
+		}
+
+		SMultiColumnTableRow<TSharedPtr<int32> >::Construct(Args, InOwnerTableView);
+	}
+
+public:
+
+	BEGIN_SLATE_FUNCTION_BUILD_OPTIMIZATION
+	virtual TSharedRef<SWidget> GenerateWidgetForColumn( const FName& ColumnName ) override
+	{
+		if (ColumnName == "JumpToButton")
+		{
+			return SNew(SButton)
+				.ToolTipText(FText::FromString("Jump to this transaction"))
+				.ButtonStyle(FEditorStyle::Get(), "FlatButton.Success")
+				.ForegroundColor(FLinearColor::White)
+				.ContentPadding(FMargin(2.f, 2.f))
+				.OnClicked_Lambda([this]() { OnGotoTransactionClicked.ExecuteIfBound(TransactionId); return FReply::Handled(); })
+				.Visibility_Lambda([this]() { return this->IsHovered() ? EVisibility::Visible : EVisibility::Hidden; })
+				[
+					SNew(SBox)
+						.HAlign(HAlign_Center)
+						[
+							SNew(STextBlock)
+								.Font(FEditorStyle::Get().GetFontStyle("FontAwesome.10"))
+								.Text(FText::FromString(TEXT("\xf138"))) /*fa-chevron-circle-right*/
+								.Justification(ETextJustify::Center)
+						]
+				];
+		}
+		else if (ColumnName == "Title")
+		{
+			return SNew(SBox)
+				.VAlign(VAlign_Center)
+				.Padding(FMargin(4.0f, 0.0f))
+				[
+					SNew(STextBlock)
+						.Text(Title)
+						.ToolTipText(Description)
+						.ColorAndOpacity(this, &SUndoHistoryTableRow::HandleTitleTextColorAndOpacity)
+				];
+		}
+
+		return SNullWidget::NullWidget;
+	}
+	END_SLATE_FUNCTION_BUILD_OPTIMIZATION
+
+private:
+
+	/** Callback for getting the color of the 'Title' text. */
+	FSlateColor HandleTitleTextColorAndOpacity( ) const
+	{
+		if (IsApplied.Get())
+		{
+			return FSlateColor::UseForeground();
+		}
+
+		return FSlateColor::UseSubduedForeground();
+	}
+
+private:
+
+	/** Holds an attribute that determines whether the transaction in this row is applied. */
+	TAttribute<bool> IsApplied;
+
+	/** Holds the transaction's index in the transaction queue. */
+	int32 QueueIndex;
+
+	/** Holds the current transaction's id. */
+	FGuid TransactionId;
+
+	/** Holds the transaction's title text. */
+	FText Title;
+
+	/** Holds the transaction's description text. */
+	FText Description;
+
+	/** Delegate called when the Goto button is clicked */
+	FOnGotoTransactionClicked OnGotoTransactionClicked;
+};
+
+
+#undef LOCTEXT_NAMESPACE
diff --git a/Engine/Source/Editor/UndoHistoryEditor/Public/IUndoHistoryEditorModule.h b/Engine/Source/Editor/UndoHistoryEditor/Public/IUndoHistoryEditorModule.h
new file mode 100644
--- /dev/null
+++ b/Engine/Source/Editor/UndoHistoryEditor/Public/IUndoHistoryEditorModule.h
@@ -0,0 +1,21 @@
+﻿// Copyright Epic Games, Inc. All Rights Reserved.
+
+#pragma once
+
+#include "CoreMinimal.h"
+#include "Modules/ModuleInterface.h"
+#include "Modules/ModuleManager.h"
+
+UNDOHISTORYEDITOR_API extern const FName UndoHistoryTabName;
+
+class UNDOHISTORYEDITOR_API IUndoHistoryEditorModule : public IModuleInterface
+{
+public:
+
+	static IUndoHistoryEditorModule& Get()
+	{
+		return FModuleManager::LoadModuleChecked<IUndoHistoryEditorModule>("UndoHistoryEditor");
+	}
+
+	virtual void ExecuteOpenUndoHistory() = 0;
+};
diff --git a/Engine/Source/Editor/UndoHistoryEditor/UndoHistoryEditor.Build.cs b/Engine/Source/Editor/UndoHistoryEditor/UndoHistoryEditor.Build.cs
new file mode 100644
--- /dev/null
+++ b/Engine/Source/Editor/UndoHistoryEditor/UndoHistoryEditor.Build.cs
@@ -0,0 +1,31 @@
+// Copyright Epic Games, Inc. All Rights Reserved.
+
+namespace UnrealBuildTool.Rules
+{
+	public class UndoHistoryEditor : ModuleRules
+	{
+		public UndoHistoryEditor(ReadOnlyTargetRules Target) : base(Target)
+		{
+			PublicDependencyModuleNames.AddRange(
+				new string[] {
+					"Core"
+				}
+			);
+
+			PrivateDependencyModuleNames.AddRange(
+				new string[] {
+					"ApplicationCore",
+					"CoreUObject",
+					"EditorFramework",
+					"EditorStyle",
+					"Engine",
+					"InputCore",
+					"Slate",
+					"SlateCore",
+					"UndoHistory",
+					"UnrealEd",
+				}
+			);
+		}
+	}
+}
