Add accessor for last undo barrier in the stack for code to query against.

# Added ability for the undo history window to add or remove undo barriers when not in PIE to facilitate testing the feature
# Fixed an issue where clearing the undo stack would result in the barriers remaining and preventing subsequent undo

diff --git a/Engine/Source/Editor/UndoHistoryEditor/Private/Widgets/SUndoHistory.cpp b/Engine/Source/Editor/UndoHistoryEditor/Private/Widgets/SUndoHistory.cpp
--- a/Engine/Source/Editor/UndoHistoryEditor/Private/Widgets/SUndoHistory.cpp
+++ b/Engine/Source/Editor/UndoHistoryEditor/Private/Widgets/SUndoHistory.cpp
@@ -94,6 +94,9 @@ void SUndoHistory::Construct( const FArguments& InArgs )
 													+ SHeaderRow::Column("Title")
 														.FillWidth(80.0f)
 
+													+ SHeaderRow::Column("UndoBarrierButton")
+														.FixedWidth(50.0f)
+
 													+ SHeaderRow::Column("JumpToButton")
 														.FixedWidth(50.0f)
 												)
@@ -323,7 +326,10 @@ TSharedRef<ITableRow> SUndoHistory::HandleUndoListGenerateRow(TSharedPtr<FTransa
 {
 	return SNew(SUndoHistoryTableRow, OwnerTable)
 		.OnGotoTransactionClicked(this, &SUndoHistory::HandleGoToTransaction)
+		.OnUndoBarrierButtonClicked(this, &SUndoHistory::HandleUndoBarrierButtonClicked)
 		.IsApplied(this, &SUndoHistory::HandleUndoListRowIsApplied, TransactionInfo->QueueIndex)
+		.IsCurrentUndoBarrier(this, &SUndoHistory::HandleIsCurrentUndoBarrier, TransactionInfo->QueueIndex)
+		.IsLastTransactionIndex(this, &SUndoHistory::HandleIsLastTransactionIndex, TransactionInfo->QueueIndex)
 		.QueueIndex(TransactionInfo->QueueIndex)
 		.Transaction(TransactionInfo->GetTransaction());
 }
@@ -345,6 +351,38 @@ void SUndoHistory::HandleGoToTransaction(const FGuid& TargetTransactionId)
 	}
 }
 
+bool SUndoHistory::HandleIsCurrentUndoBarrier(int32 InQueueIndex) const
+{
+	if ((GEditor == nullptr) || (GEditor->Trans == nullptr))
+	{
+		return false;
+	}
+
+	return InQueueIndex == GEditor->Trans->GetCurrentUndoBarrier() - 1;
+}
+
+bool SUndoHistory::HandleIsLastTransactionIndex(int32 InQueueIndex) const
+{
+	if ((GEditor == nullptr) || (GEditor->Trans == nullptr))
+	{
+		return false;
+	}
+
+	return InQueueIndex == (GEditor->Trans->GetQueueLength() - GEditor->Trans->GetUndoCount() - 1);
+}
+
+void SUndoHistory::HandleUndoBarrierButtonClicked(int32 InQueueIndex)
+{
+	if (HandleIsCurrentUndoBarrier(InQueueIndex))
+	{
+		GEditor->Trans->RemoveUndoBarrier();
+	}
+	else if (HandleIsLastTransactionIndex(InQueueIndex))
+	{
+		GEditor->Trans->SetUndoBarrier();
+	}
+}
+
 bool SUndoHistory::HandleUndoListRowIsApplied(int32 QueueIndex) const
 {
 	if ((GEditor == nullptr) || (GEditor->Trans == nullptr))
diff --git a/Engine/Source/Editor/UndoHistoryEditor/Private/Widgets/SUndoHistory.h b/Engine/Source/Editor/UndoHistoryEditor/Private/Widgets/SUndoHistory.h
--- a/Engine/Source/Editor/UndoHistoryEditor/Private/Widgets/SUndoHistory.h
+++ b/Engine/Source/Editor/UndoHistoryEditor/Private/Widgets/SUndoHistory.h
@@ -73,6 +73,15 @@ private:
 	/** Callback for when a user wants to jump to a certain transaction. */ 
 	void HandleGoToTransaction(const FGuid& TargetTransactionId);
 
+	/** Callback for checking if the given index is the current undo barrier index. */
+	bool HandleIsCurrentUndoBarrier(int32 InQueueIndex) const;
+
+	/** Callback for checking if the given index is the last active transaction. */
+	bool HandleIsLastTransactionIndex(int32 InQueueIndex) const;
+
+	/** Callback for when a user clicks to add or remove an undo barrier. */
+	void HandleUndoBarrierButtonClicked(int32 InQueueIndex);
+
 	/** Callback for checking whether the specified undo list row transaction is applied. */
 	bool HandleUndoListRowIsApplied( int32 QueueIndex ) const;
 
diff --git a/Engine/Source/Editor/UndoHistoryEditor/Private/Widgets/SUndoHistoryTableRow.h b/Engine/Source/Editor/UndoHistoryEditor/Private/Widgets/SUndoHistoryTableRow.h
--- a/Engine/Source/Editor/UndoHistoryEditor/Private/Widgets/SUndoHistoryTableRow.h
+++ b/Engine/Source/Editor/UndoHistoryEditor/Private/Widgets/SUndoHistoryTableRow.h
@@ -22,10 +22,12 @@
 #include "SPositiveActionButton.h"
 #include "Widgets/Views/STreeView.h"
 #include "Widgets/Input/SButton.h"
+#include "SNegativeActionButton.h"
 
 #define LOCTEXT_NAMESPACE "SUndoHistoryTableRow"
 
 DECLARE_DELEGATE_OneParam(FOnGotoTransactionClicked, const FGuid&)
+DECLARE_DELEGATE_OneParam(FOnUndoBarrierButtonClicked, int32)
 
 /**
  * Implements a row widget for the undo history list.
@@ -38,9 +40,12 @@ public:
 
 	SLATE_BEGIN_ARGS(SUndoHistoryTableRow) { }
 		SLATE_ATTRIBUTE(bool, IsApplied)
+		SLATE_ATTRIBUTE(bool, IsCurrentUndoBarrier)
+		SLATE_ATTRIBUTE(bool, IsLastTransactionIndex)
 		SLATE_ARGUMENT(int32, QueueIndex)
 		SLATE_ARGUMENT(const FTransaction*, Transaction)
 		SLATE_EVENT(FOnGotoTransactionClicked, OnGotoTransactionClicked)
+		SLATE_EVENT(FOnUndoBarrierButtonClicked, OnUndoBarrierButtonClicked)
 	SLATE_END_ARGS()
 
 public:
@@ -53,8 +58,11 @@ public:
 	void Construct( const FArguments& InArgs, const TSharedRef<STableViewBase>& InOwnerTableView )
 	{
 		IsApplied = InArgs._IsApplied;
+		IsCurrentUndoBarrier = InArgs._IsCurrentUndoBarrier;
+		IsLastTransactionIndex = InArgs._IsLastTransactionIndex;
 		QueueIndex = InArgs._QueueIndex;
 		OnGotoTransactionClicked = InArgs._OnGotoTransactionClicked;
+		OnUndoBarrierButtonClicked = InArgs._OnUndoBarrierButtonClicked;
 
 		FSuperRowType::FArguments Args = FSuperRowType::FArguments()
 			.Style(&FAppStyle::Get().GetWidgetStyle<FTableRowStyle>("SceneOutliner.TableViewRow"));
@@ -101,6 +109,14 @@ public:
 						.ColorAndOpacity(this, &SUndoHistoryTableRow::HandleTitleTextColorAndOpacity)
 				];
 		}
+		else if (ColumnName == "UndoBarrierButton")
+		{
+			return SNew(SNegativeActionButton)
+				.ToolTipText_Lambda( [this]() { if (IsCurrentUndoBarrier.Get()) { return LOCTEXT("RemoveUndoBarrierTooltip", "Remove Undo Barrier"); } return LOCTEXT("SetUndoBarrierTooltip", "Set Undo Barrier"); })
+				.Icon(FAppStyle::GetBrush("Icons.PlusCircle"))
+				.OnClicked_Lambda([this]() { OnUndoBarrierButtonClicked.ExecuteIfBound(QueueIndex); return FReply::Handled(); })
+				.Visibility_Lambda([this]() { return (GEditor && (!GEditor->PlayWorld || GEditor->bIsSimulatingInEditor)) && this->IsHovered() && (IsCurrentUndoBarrier.Get() || IsLastTransactionIndex.Get()) ? EVisibility::Visible : EVisibility::Hidden; });
+		}
 
 		return SNullWidget::NullWidget;
 	}
@@ -124,6 +140,12 @@ private:
 	/** Holds an attribute that determines whether the transaction in this row is applied. */
 	TAttribute<bool> IsApplied;
 
+	/** Holds an attribute that determines whether the transaction in this row is the current undo barrier. */
+	TAttribute<bool> IsCurrentUndoBarrier;
+
+	/** Holds an attribute that determines whether the transaction in this row is the last active transaction. */
+	TAttribute<bool> IsLastTransactionIndex;
+
 	/** Holds the transaction's index in the transaction queue. */
 	int32 QueueIndex;
 
@@ -138,6 +160,9 @@ private:
 
 	/** Delegate called when the Goto button is clicked */
 	FOnGotoTransactionClicked OnGotoTransactionClicked;
+
+	/** Delegate called when the Barrier button is clicked */
+	FOnUndoBarrierButtonClicked OnUndoBarrierButtonClicked;
 };
 
 
diff --git a/Engine/Source/Editor/UnrealEd/Classes/Editor/TransBuffer.h b/Engine/Source/Editor/UnrealEd/Classes/Editor/TransBuffer.h
--- a/Engine/Source/Editor/UnrealEd/Classes/Editor/TransBuffer.h
+++ b/Engine/Source/Editor/UnrealEd/Classes/Editor/TransBuffer.h
@@ -159,6 +159,7 @@ public:
 	virtual void SetUndoBarrier() override;
 	virtual void RemoveUndoBarrier() override;
 	virtual void ClearUndoBarriers() override;
+	virtual int32 GetCurrentUndoBarrier() const override;
 	virtual bool Undo(bool bCanRedo = true) override;
 	virtual bool Redo() override;
 	virtual bool EnableObjectSerialization() override;
diff --git a/Engine/Source/Editor/UnrealEd/Classes/Editor/Transactor.h b/Engine/Source/Editor/UnrealEd/Classes/Editor/Transactor.h
--- a/Engine/Source/Editor/UnrealEd/Classes/Editor/Transactor.h
+++ b/Engine/Source/Editor/UnrealEd/Classes/Editor/Transactor.h
@@ -588,6 +588,11 @@ class UNREALED_API UTransactor : public UObject
 	 */
 	virtual void ClearUndoBarriers() PURE_VIRTUAL(UTransactor::ClearUndoBarriers, );
 
+	/**
+	 * Gets the current undo barrier's position in transaction queue
+	 */
+	 virtual int32 GetCurrentUndoBarrier() const PURE_VIRTUAL(UTransactor::GetCurrentUndoBarrier(), { return INDEX_NONE; });
+
 	/**
 	 * Executes an undo transaction, undoing all actions contained by that transaction.
 	 * 
diff --git a/Engine/Source/Editor/UnrealEd/Private/EditorTransaction.cpp b/Engine/Source/Editor/UnrealEd/Private/EditorTransaction.cpp
--- a/Engine/Source/Editor/UnrealEd/Private/EditorTransaction.cpp
+++ b/Engine/Source/Editor/UnrealEd/Private/EditorTransaction.cpp
@@ -1150,6 +1150,7 @@ void UTransBuffer::Reset( const FText& Reason )
 		ResetReason = Reason;
 		ActiveCount = 0;
 		ActiveRecordCounts.Empty();
+		ClearUndoBarriers();
 		UndoBufferChangedDelegate.Broadcast();
 
 		CheckState();
@@ -1350,6 +1351,16 @@ void UTransBuffer::ClearUndoBarriers()
 	UndoBarrierStack.Empty();
 }
 
+int32 UTransBuffer::GetCurrentUndoBarrier() const
+{
+	if (UndoBarrierStack.Num() > 0)
+	{
+		return UndoBarrierStack.Last();
+	}
+
+	return INDEX_NONE;
+}
+
 
 bool UTransBuffer::Undo(bool bCanRedo)
 {
