Undo history: properly cleaned up the transaction details view when clearing the undo stack

diff --git a/Engine/Source/Editor/UndoHistoryEditor/Private/Widgets/SUndoHistory.cpp b/Engine/Source/Editor/UndoHistoryEditor/Private/Widgets/SUndoHistory.cpp
--- a/Engine/Source/Editor/UndoHistoryEditor/Private/Widgets/SUndoHistory.cpp
+++ b/Engine/Source/Editor/UndoHistoryEditor/Private/Widgets/SUndoHistory.cpp
@@ -302,6 +302,10 @@ void SUndoHistory::SelectLastTransaction()
 				UndoDetailsView->SetSelectedTransaction(InnerTransaction->GenerateDiff());
 			}
 		}
+		else
+		{
+			UndoDetailsView->Reset();
+		}
 	}
 }
 
