[NO NON-PII COMMIT MESSAGE]

diff --git a/Engine/Source/Editor/UndoHistoryEditor/Private/Widgets/SUndoHistory.cpp b/Engine/Source/Editor/UndoHistoryEditor/Private/Widgets/SUndoHistory.cpp
--- a/Engine/Source/Editor/UndoHistoryEditor/Private/Widgets/SUndoHistory.cpp
+++ b/Engine/Source/Editor/UndoHistoryEditor/Private/Widgets/SUndoHistory.cpp
@@ -483,7 +483,14 @@ FText SUndoHistory::HandleUndoSizeTextBlockText() const
 		return FText::GetEmpty();
 	}
 
-	return FText::Format(LOCTEXT("TransactionCountF", "{0} Transactions ({1})"), FText::AsNumber(UndoList.Num()), FText::AsMemory(GEditor->Trans->GetUndoSize()));
+	if (UndoList.Num() <= 1)
+	{
+		return FText::Format(LOCTEXT("SingleTransactionF", "{0} Transaction ({1})"), FText::AsNumber(UndoList.Num()), FText::AsMemory(GEditor->Trans->GetUndoSize()));
+	}
+	else
+	{
+		return FText::Format(LOCTEXT("TransactionCountF", "{0} Transactions ({1})"), FText::AsNumber(UndoList.Num()), FText::AsMemory(GEditor->Trans->GetUndoSize()));
+	}
 }
 
 EVisibility SUndoHistory::HandleUndoWarningVisibility() const
