Add plurality token to the "Transaction" test on the undo history window

diff --git a/Engine/Source/Editor/UndoHistoryEditor/Private/Widgets/SUndoHistory.cpp b/Engine/Source/Editor/UndoHistoryEditor/Private/Widgets/SUndoHistory.cpp
--- a/Engine/Source/Editor/UndoHistoryEditor/Private/Widgets/SUndoHistory.cpp
+++ b/Engine/Source/Editor/UndoHistoryEditor/Private/Widgets/SUndoHistory.cpp
@@ -483,14 +483,7 @@ FText SUndoHistory::HandleUndoSizeTextBlockText() const
 		return FText::GetEmpty();
 	}
 
-	if (UndoList.Num() <= 1)
-	{
-		return FText::Format(LOCTEXT("SingleTransactionF", "{0} Transaction ({1})"), FText::AsNumber(UndoList.Num()), FText::AsMemory(GEditor->Trans->GetUndoSize()));
-	}
-	else
-	{
-		return FText::Format(LOCTEXT("TransactionCountF", "{0} Transactions ({1})"), FText::AsNumber(UndoList.Num()), FText::AsMemory(GEditor->Trans->GetUndoSize()));
-	}
+	return FText::Format(LOCTEXT("TransactionCountF", "{0} {0}|plural(one=Transaction,other=Transactions) ({1})"), FText::AsNumber(UndoList.Num()), FText::AsMemory(GEditor->Trans->GetUndoSize()));
 }
 
 EVisibility SUndoHistory::HandleUndoWarningVisibility() const
