[Localization] "Jump to this transaction" button in Undo History is now set for localization

diff --git a/Engine/Source/Editor/UndoHistoryEditor/Private/Widgets/SUndoHistoryTableRow.h b/Engine/Source/Editor/UndoHistoryEditor/Private/Widgets/SUndoHistoryTableRow.h
--- a/Engine/Source/Editor/UndoHistoryEditor/Private/Widgets/SUndoHistoryTableRow.h
+++ b/Engine/Source/Editor/UndoHistoryEditor/Private/Widgets/SUndoHistoryTableRow.h
@@ -93,7 +93,7 @@ public:
 		if (ColumnName == "JumpToButton")
 		{
 			return SNew(SPositiveActionButton)
-				.ToolTipText(FText::FromString("Jump to this transaction"))
+				.ToolTipText(LOCTEXT("JumpToTransactionTooltip", "Jump to this transaction"))
 				.Icon(FAppStyle::GetBrush("Icons.CircleArrowRight"))
 				.OnClicked_Lambda([this]() { OnGotoTransactionClicked.ExecuteIfBound(TransactionId); return FReply::Handled(); })
 				.Visibility_Lambda([this]() { return this->IsHovered() ? EVisibility::Visible : EVisibility::Hidden; });
