Fix crash in SUndoHistory when shutting down the engine.

diff --git a/Engine/Source/Editor/UndoHistoryEditor/Private/Widgets/SUndoHistory.cpp b/Engine/Source/Editor/UndoHistoryEditor/Private/Widgets/SUndoHistory.cpp
--- a/Engine/Source/Editor/UndoHistoryEditor/Private/Widgets/SUndoHistory.cpp
+++ b/Engine/Source/Editor/UndoHistoryEditor/Private/Widgets/SUndoHistory.cpp
@@ -231,9 +231,11 @@ SUndoHistory::~SUndoHistory()
 
 void SUndoHistory::OnUndoBufferChanged()
 {
-	ReloadUndoList();
-
-	SelectLastTransaction();
+	if (!IsEngineExitRequested())
+	{
+		ReloadUndoList();
+		SelectLastTransaction();
+	}
 }
 
 
